Fixed up 4.10 TP_SideScroller (C++) template by removing the UPROPERTY tag from the offending variable and then resaving the map.

It is likely that there is still an underlying problem here. I'm not clear that what was done should cause an error message, given that it worked fine. So, either something is broken in an ignorable way, or the error message is coming up when it shouldn't.

diff --git a/Templates/TP_SideScroller/Source/TP_SideScroller/TP_SideScrollerCharacter.h b/Templates/TP_SideScroller/Source/TP_SideScroller/TP_SideScrollerCharacter.h
--- a/Templates/TP_SideScroller/Source/TP_SideScroller/TP_SideScrollerCharacter.h
+++ b/Templates/TP_SideScroller/Source/TP_SideScroller/TP_SideScrollerCharacter.h
@@ -9,7 +9,6 @@ class ATP_SideScrollerCharacter : public ACharacter
 	GENERATED_BODY()
 
 	/** Side view camera */
-	UPROPERTY(VisibleAnywhere, BlueprintReadOnly, Category = Camera, meta = (AllowPrivateAccess = "true"))
 	class UCameraComponent* SideViewCameraComponent;
 
 	/** Camera boom positioning the camera beside the character */
